PR #9538: Added a plugin to refresh Mainframe and Slate Widgets when Live Coding is complete. (Contributed by [REDACTED_PERSON])

diff --git a/Engine/Plugins/Editor/LiveUpdateForSlate/LiveUpdateForSlate.uplugin b/Engine/Plugins/Editor/LiveUpdateForSlate/LiveUpdateForSlate.uplugin
new file mode 100644
--- /dev/null
+++ b/Engine/Plugins/Editor/LiveUpdateForSlate/LiveUpdateForSlate.uplugin
@@ -0,0 +1,24 @@
+{
+	"FileVersion": 3,
+	"Version": 1,
+	"VersionName": "1.0",
+	"FriendlyName": "Live Update for Slate",
+	"Description": "Refreshes the editor layout and tabs when Live Coding is complete",
+	"Category": "Other",
+	"CreatedBy": "Alessa Baker & Epic Games, Inc.",
+	"CreatedByURL": "codekitten.me",
+	"DocsURL": "",
+	"MarketplaceURL": "",
+	"SupportURL": "",
+	"CanContainContent": true,
+	"IsBetaVersion": false,
+	"IsExperimentalVersion": false,
+	"Installed": false,
+	"Modules": [
+		{
+			"Name": "LiveUpdateForSlate",
+			"Type": "Editor",
+			"LoadingPhase": "Default"
+		}
+	]
+}
\ No newline at end of file
diff --git a/Engine/Plugins/Editor/LiveUpdateForSlate/Source/LiveUpdateForSlate/LiveUpdateForSlate.Build.cs b/Engine/Plugins/Editor/LiveUpdateForSlate/Source/LiveUpdateForSlate/LiveUpdateForSlate.Build.cs
new file mode 100644
--- /dev/null
+++ b/Engine/Plugins/Editor/LiveUpdateForSlate/Source/LiveUpdateForSlate/LiveUpdateForSlate.Build.cs
@@ -0,0 +1,37 @@
+// Copyright Epic Games, Inc. All Rights Reserved.
+
+using UnrealBuildTool;
+
+public class LiveUpdateForSlate : ModuleRules
+{
+	public LiveUpdateForSlate(ReadOnlyTargetRules Target) : base(Target)
+	{
+		PCHUsage = PCHUsageMode.UseExplicitOrSharedPCHs;
+		
+		PublicDependencyModuleNames.AddRange(
+			new string[]
+			{
+				"Core",
+			}
+		);
+		
+		PrivateDependencyModuleNames.AddRange(
+			new string[]
+			{
+				"CoreUObject",
+				"Engine",
+				"MainFrame",
+				"Settings",
+				"SettingsEditor",
+				"Slate",
+				"SlateCore",
+				"UnrealEd",
+			}
+		);
+
+		if (Target.bWithLiveCoding)
+		{
+			PrivateDependencyModuleNames.Add("LiveCoding");
+		}
+	}
+}
diff --git a/Engine/Plugins/Editor/LiveUpdateForSlate/Source/LiveUpdateForSlate/Private/LiveUpdateForSlate.cpp b/Engine/Plugins/Editor/LiveUpdateForSlate/Source/LiveUpdateForSlate/Private/LiveUpdateForSlate.cpp
new file mode 100644
--- /dev/null
+++ b/Engine/Plugins/Editor/LiveUpdateForSlate/Source/LiveUpdateForSlate/Private/LiveUpdateForSlate.cpp
@@ -0,0 +1,84 @@
+// Copyright Epic Games, Inc. All Rights Reserved.
+
+#include "LiveUpdateForSlate.h"
+#include "Editor.h"
+#include "Framework/Docking/TabManager.h"
+#include "Interfaces/IMainFrameModule.h"
+#include "ISettingsModule.h"
+#include "LiveUpdateSlateSettings.h"
+#include "Misc/ConfigCacheIni.h"
+#include "Subsystems/AssetEditorSubsystem.h"
+
+#if WITH_LIVE_CODING
+#include "ILiveCodingModule.h"
+#endif
+
+#define LOCTEXT_NAMESPACE "LiveUpdateForSlateModule"
+
+IMPLEMENT_MODULE(FLiveUpdateForSlateModule, LiveUpdateForSlate)
+
+void FLiveUpdateForSlateModule::StartupModule()
+{
+	if (ISettingsModule* SettingsModule = FModuleManager::GetModulePtr<ISettingsModule>("Settings"))
+	{
+		SettingsModule->RegisterSettings("Editor", "Plugins", "Live Update for Slate",
+			LOCTEXT("LiveUpdateSlateSettingsName", "Live Update for Slate"),
+			LOCTEXT("LiveUpdateSlateSettingsDescription", "Live Update for Slate"),
+			GetMutableDefault<ULiveUpdateSlateSettings>()
+		);
+	}
+
+#if WITH_LIVE_CODING
+	if (ILiveCodingModule* LiveCoding = FModuleManager::LoadModulePtr<ILiveCodingModule>(LIVE_CODING_MODULE_NAME))
+	{
+		OnPatchCompleteHandle = LiveCoding->GetOnPatchCompleteDelegate().AddRaw(this, &FLiveUpdateForSlateModule::OnPatchComplete);
+	}
+#endif
+}
+
+void FLiveUpdateForSlateModule::ShutdownModule()
+{
+	if (ISettingsModule* SettingsModule = FModuleManager::GetModulePtr<ISettingsModule>("Settings"))
+	{
+		SettingsModule->UnregisterSettings("Editor", "Plugins", "Live Update for Slate");
+	}
+
+#if WITH_LIVE_CODING
+	if (ILiveCodingModule* LiveCoding = FModuleManager::LoadModulePtr<ILiveCodingModule>(LIVE_CODING_MODULE_NAME))
+	{
+		LiveCoding->GetOnPatchCompleteDelegate().Remove(OnPatchCompleteHandle);
+	}
+#endif
+}
+
+void FLiveUpdateForSlateModule::OnPatchComplete()
+{
+	ULiveUpdateSlateSettings* Settings = GetMutableDefault<ULiveUpdateSlateSettings>();
+
+	// Don't rebuild Slate if the plugin is not enabled
+	if (!Settings->bEnableLiveUpdateForSlate)
+	{
+		return;
+	}
+
+	// Save all open tabs
+	FGlobalTabmanager::Get()->SaveAllVisualState();
+
+	// Flush our config files
+	const bool bRead = true;
+	GConfig->Flush(bRead, GEditorLayoutIni);
+
+	UAssetEditorSubsystem* Subsystem = GEditor->GetEditorSubsystem<UAssetEditorSubsystem>();
+	TArray<UObject*> OpenedAssets = Subsystem->GetAllEditedAssets();
+
+	// Request the editor to rebuild Slate
+	IMainFrameModule& MainFrameModule = FModuleManager::LoadModuleChecked<IMainFrameModule>("MainFrame");
+	MainFrameModule.RecreateDefaultMainFrame(false, false);
+
+	for (const UObject* Asset : OpenedAssets)
+	{
+		Subsystem->OpenEditorForAsset(Asset);
+	}
+}
+
+#undef LOCTEXT_NAMESPACE
diff --git a/Engine/Plugins/Editor/LiveUpdateForSlate/Source/LiveUpdateForSlate/Public/LiveUpdateForSlate.h b/Engine/Plugins/Editor/LiveUpdateForSlate/Source/LiveUpdateForSlate/Public/LiveUpdateForSlate.h
new file mode 100644
--- /dev/null
+++ b/Engine/Plugins/Editor/LiveUpdateForSlate/Source/LiveUpdateForSlate/Public/LiveUpdateForSlate.h
@@ -0,0 +1,20 @@
+// Copyright Epic Games, Inc. All Rights Reserved.
+
+#pragma once
+
+#include "Delegates/IDelegateInstance.h"
+#include "Modules/ModuleManager.h"
+
+class FLiveUpdateForSlateModule : public IModuleInterface
+{
+public:
+
+	/** IModuleInterface implementation */
+	virtual void StartupModule() override;
+	virtual void ShutdownModule() override;
+
+private:
+	void OnPatchComplete();
+
+	FDelegateHandle OnPatchCompleteHandle;
+};
diff --git a/Engine/Plugins/Editor/LiveUpdateForSlate/Source/LiveUpdateForSlate/Public/LiveUpdateSlateSettings.h b/Engine/Plugins/Editor/LiveUpdateForSlate/Source/LiveUpdateForSlate/Public/LiveUpdateSlateSettings.h
new file mode 100644
--- /dev/null
+++ b/Engine/Plugins/Editor/LiveUpdateForSlate/Source/LiveUpdateForSlate/Public/LiveUpdateSlateSettings.h
@@ -0,0 +1,16 @@
+// Copyright Epic Games, Inc. All Rights Reserved.
+
+#pragma once
+
+#include "UObject/Object.h"
+#include "LiveUpdateSlateSettings.generated.h"
+
+UCLASS(config = EditorPerProjectUserSettings)
+class LIVEUPDATEFORSLATE_API ULiveUpdateSlateSettings : public UObject
+{
+	GENERATED_BODY()
+
+public:
+	UPROPERTY(Config, EditAnywhere, Category=Slate, meta=(ToolTip="Refreshes the editor's Slate layout when Live Coding patches are complete."))
+	bool bEnableLiveUpdateForSlate = true;
+};
